Fix CIS errors

diff --git a/Engine/Plugins/Experimental/CharacterAI/Source/CharacterAI/Private/CharacterAIModule.cpp b/Engine/Plugins/Experimental/CharacterAI/Source/CharacterAI/Private/CharacterAIModule.cpp
--- a/Engine/Plugins/Experimental/CharacterAI/Source/CharacterAI/Private/CharacterAIModule.cpp
+++ b/Engine/Plugins/Experimental/CharacterAI/Source/CharacterAI/Private/CharacterAIModule.cpp
@@ -2,6 +2,7 @@
 
 #include "CharacterAIPrivatePCH.h"
 #include "CharacterAIModule.h"
+#include "ModuleManager.h"
 
 //////////////////////////////////////////////////////////////////////////
 // FCharacterAIModule
diff --git a/Engine/Plugins/Experimental/CharacterAI/Source/CharacterAI/Private/CharacterAIPrivatePCH.h b/Engine/Plugins/Experimental/CharacterAI/Source/CharacterAI/Private/CharacterAIPrivatePCH.h
--- a/Engine/Plugins/Experimental/CharacterAI/Source/CharacterAI/Private/CharacterAIPrivatePCH.h
+++ b/Engine/Plugins/Experimental/CharacterAI/Source/CharacterAI/Private/CharacterAIPrivatePCH.h
@@ -5,7 +5,7 @@
 // Uncommenting this like causes the FTransform linkage issue
 //#include "Core.h"
 
-//#include "Engine.h"
+#include "Engine.h"
 
 DECLARE_LOG_CATEGORY_EXTERN(LogCharacterAI, Verbose, All);
 
diff --git a/Engine/Plugins/Experimental/CharacterAI/Source/CharacterAI/Public/CharacterAIModule.h b/Engine/Plugins/Experimental/CharacterAI/Source/CharacterAI/Public/CharacterAIModule.h
--- a/Engine/Plugins/Experimental/CharacterAI/Source/CharacterAI/Public/CharacterAIModule.h
+++ b/Engine/Plugins/Experimental/CharacterAI/Source/CharacterAI/Public/CharacterAIModule.h
@@ -2,6 +2,8 @@
 
 #pragma once
 
+#include "ModuleInterface.h"
+
 #define CHARACTERAI_MODULE_NAME "CharacterAI"
 
 //////////////////////////////////////////////////////////////////////////
