Mark a couple of plugins as beta, rather than having [EARLY ACCESS PREVIEW] in the title.

diff --git a/Engine/Plugins/Experimental/CharacterAI/CharacterAI.uplugin b/Engine/Plugins/Experimental/CharacterAI/CharacterAI.uplugin
--- a/Engine/Plugins/Experimental/CharacterAI/CharacterAI.uplugin
+++ b/Engine/Plugins/Experimental/CharacterAI/CharacterAI.uplugin
@@ -7,9 +7,10 @@
 	"CreatedBy" : "Epic Games, Inc.",
 	"CreatedByURL" : "http://epicgames.com",
 	"EngineVersion" : "4.5.0",
-	"Description" : "[EARLY ACCESS PREVIEW] CharacterAI adds code and assets related to implementing AI in a character-based project.",
+	"Description" : "Adds code and assets related to implementing AI in a character-based project.",
 	"Category" : "Gameplay",
-	"EnabledByDefault" : true,
+    "EnabledByDefault": true,
+    "IsBetaVersion": true,
 
 	"Modules" :
 	[
diff --git a/Engine/Plugins/Runtime/TwitchLiveStreaming/TwitchLiveStreaming.uplugin b/Engine/Plugins/Runtime/TwitchLiveStreaming/TwitchLiveStreaming.uplugin
--- a/Engine/Plugins/Runtime/TwitchLiveStreaming/TwitchLiveStreaming.uplugin
+++ b/Engine/Plugins/Runtime/TwitchLiveStreaming/TwitchLiveStreaming.uplugin
@@ -7,9 +7,10 @@
 	"CreatedBy" : "Epic Games, Inc.",
 	"CreatedByURL" : "http://epicgames.com",
 	"EngineVersion" : "4.4.0",
-	"Description" : "[EARLY ACCESS PREVIEW] Live video broadcasting to Twitch for games and the editor",
+	"Description" : "Live video broadcasting to Twitch for games and the editor",
 	"Category" : "Misc",
 	"EnabledByDefault" : false,
+    "IsBetaVersion": true,
 	
 	"Modules" :
 	[
